Misc fixes for build system in //UE4/Main:

* Adding exclusions for TCL post processor matching Android output as an error
* Ignore tempnam warning on Linux caused by linking against libfbx.a
* Ignore LNK4221 warning in non-unity compiles on Windows caused by cpp files which have no symbols (typically because they're #ifdef'ed out)
* Specify multiple recipients for emails in a list rather than as a space-delimited string, fixing error trying to send notifications.

diff --git a/Engine/Build/BuildFarm/PostpExtensions.pl b/Engine/Build/BuildFarm/PostpExtensions.pl
--- a/Engine/Build/BuildFarm/PostpExtensions.pl
+++ b/Engine/Build/BuildFarm/PostpExtensions.pl
@@ -1,5 +1,15 @@
 # These are patterns we want to ignore
 @::gExcludePatterns = (
+	# UAT spam that just adds noise
+	".*Exception in AutomationTool: BUILD FAILED.*",
+	"RunUAT ERROR: AutomationTool was unable to run successfully.",
+
+	# Linux error caused by linking against libfbx
+	".*the use of `tempnam' is dangerous.*",
+
+	# Warning about object files not definining any public symbols in non-unity builds
+	".*warning LNK4221:.*",
+
 #	".*ERROR: The process.*not found",
 #	".*ERROR: This operation returned because the timeout period expired.*",
 #	".*Sync.VerifyKnownFileInManifest: ERROR:.*",
@@ -17,11 +27,8 @@
 #	".*xgConsole: BuildSystem failed to start with error code 217.*",
 #	".*error: Unhandled error code.*",
 #	".*621InnerException.*",
-	".*Exception in AutomationTool: BUILD FAILED.*",
-	"RunUAT ERROR: AutomationTool was unable to run successfully.",
 #	".* GUBP.PrintDetailedChanges:.*",
 #	".* Failed to produce item:.*",
-#	".*the use of `tempnam' is dangerous.*",
 #	".*major version 51 is newer than 50.*",
 #	"Cannot find project '.+' in branch",
 #	".*==== Writing to template target file.*",
@@ -34,6 +41,10 @@
 #	".*SignTool Error: An error occurred while attempting to sign:.*",
 ); 
 
+# TCL rules match any lines beginning with '====', which are output by Android toolchain.
+$::gDontCheck .= "," if $::gDontCheck;
+$::gDontCheck .= "tclTestFail,tclClusterTestTimeout";
+
 # These are patterns we want to process
 # NOTE: order is important because the file is processed line by line 
 # After an error is processed on a line that line is considered processed
