Mask a few more warnings from the post-processor: XboxOne link warnings due to missing PDBs, and Android warnings during Java compilation about deprecated APIs and compiler versions.

diff --git a/Engine/Build/BuildFarm/PostpExtensions.pl b/Engine/Build/BuildFarm/PostpExtensions.pl
--- a/Engine/Build/BuildFarm/PostpExtensions.pl
+++ b/Engine/Build/BuildFarm/PostpExtensions.pl
@@ -9,7 +9,13 @@
 
 	# Warning about object files not definining any public symbols in non-unity builds
 	".*warning LNK4221:.*",
+	
+	# Warning about static libraries missing PDB files for XboxOne
+	".*warning LNK4099:.*",
 
+	# Warning from Android monolithic builds about the javac version 
+	".*major version 51 is newer than 50.*",
+	
 #	".*ERROR: The process.*not found",
 #	".*ERROR: This operation returned because the timeout period expired.*",
 #	".*Sync.VerifyKnownFileInManifest: ERROR:.*",
@@ -29,7 +35,6 @@
 #	".*621InnerException.*",
 #	".* GUBP.PrintDetailedChanges:.*",
 #	".* Failed to produce item:.*",
-#	".*major version 51 is newer than 50.*",
 #	"Cannot find project '.+' in branch",
 #	".*==== Writing to template target file.*",
 #	".*==== Writing to OBB data file.*",
@@ -45,6 +50,9 @@
 $::gDontCheck .= "," if $::gDontCheck;
 $::gDontCheck .= "tclTestFail,tclClusterTestTimeout";
 
+# We don't compile clean for Android, so ignore javac notices for now
+$::gDontCheck .= ",javacNote";
+
 # These are patterns we want to process
 # NOTE: order is important because the file is processed line by line 
 # After an error is processed on a line that line is considered processed
