Improve post-processing of errors and warnings for Mac cooks and localization gathers, and fix undefined variable warnings in lots of the default Epic matchers.

diff --git a/Engine/Build/BuildFarm/PostpExtensions.pl b/Engine/Build/BuildFarm/PostpExtensions.pl
--- a/Engine/Build/BuildFarm/PostpExtensions.pl
+++ b/Engine/Build/BuildFarm/PostpExtensions.pl
@@ -95,6 +95,11 @@ unshift @::gMatchers, (
 		id =>               "editorLogChannelDisplay",
 		pattern =>          q{^([a-zA-Z_][a-zA-Z0-9_]*):Display: },
 		action =>           q{ }
+	},
+	{
+		id =>               "automationException",
+		pattern =>          q{^.* AutomationTool.AutomationException: },
+		action =>           q{incValue("errors"); diagnostic("Exception", "error", 0, forwardWhile("^  at "));}
 	}
 );
 
@@ -102,207 +107,72 @@ push @::gMatchers, (
 	{
 		id => "AutomationToolException",
 		pattern => q{AutomationTool terminated with exception:},
-		action => q{incValue("errors"); diagnostic("$1", "error", 0,  forwardTo(q{AutomationTool exiting with ExitCode=}));}
-	},	
-	{
-		id => "UnableToConnectToVerisignSigningServer",
-		pattern => q{Failed to sign executables .+ times over a period of .+},
-		action => q{incValue("errors"); diagnostic("$1", "error", 0,  forwardTo(q{returned an invalid response.}));}
+		action => q{incValue("errors"); diagnostic("Exception", "error", 0,  forwardTo(q{AutomationTool exiting with ExitCode=}));}
 	},	
-	{
-		id => "VerisignTimestampServerConnectionError",
-		pattern => q{The specified timestamp server either could not be reached or},
-		action => q{incValue("errors"); diagnostic("$1", "error", 0, 5);}
-	},
-	{
-		id => "OverlappingBuildProducts",
-		pattern => q{Overlapping build product:},
-		action => q{incValue("errors"); diagnostic("$1", "error", 0, forwardTo(q{Stacktrace:}));}
-	},
-	{
-		id => "CookCommandlet",
-		pattern => q{:Error:},
-		action => q{incValue("errors"); diagnostic("$1", "error", 0, forwardTo(q{0}));}
-	},
-	{
-		id => "clang",
-		pattern => q{clang: error:},
-		action => q{incValue("errors"); diagnostic("$1", "error", -5, forwardTo(q{0}));}
-	},
 	{
 		id => "symbols",
 		pattern => q{Undefined symbols for architecture},
-		action => q{incValue("errors"); diagnostic("$1", "error", 0, forwardTo(q{ld: symbol}));}
-	},
-	{
-		id => "couldNotOpen",
-		pattern => q{Error: Could not open:},
-		action => q{incValue("errors"); diagnostic("$1", "error", 0, forwardTo(q{0}));}
-	},
-	{
-		id => "errorError",
-		pattern => q{Error: Error},
-		action => q{incValue("errors"); diagnostic("$1", "error", 0, forwardTo(q{0}));}
-	},
-	{
-		id => "APIDocFail",
-		pattern => q{Unhandled Exception: System.IO.DirectoryNotFoundException:},
-		action => q{incValue("errors"); diagnostic("$1", "error", 0, 0);}
-	},
-	{
-		id => "ProjFind",
-		pattern => q{Exception in AutomationTool: Cannot find project},
-		action => q{incValue("errors"); diagnostic("$1", "error", 0, 0);}
-	},
-	{
-		id => "buildMissing",
-		pattern => q{build is missing or did not complete},
-		action => q{incValue("errors"); diagnostic("$1", "error", 0, 0);}
-	},
-	{
-		id => "duplicateNode",
-		pattern => q{Attempt to add a duplicate node},
-		action => q{incValue("errors"); diagnostic("$1", "error", 0, 0);}
+		action => q{incValue("errors"); diagnostic("Symbols", "error", 0, forwardTo(q{ld: symbol}));}
 	},
 	{
 		id => "LogError",
 		pattern =>q{LogWindows: === Critical error: ===},
-		action => q{incValue("errors"); diagnostic("$1", "error", 0, forwardTo(q{Executing StaticShutdownAfterError}));}
-	},
-	{
-		id => "tempStore",
-		pattern => q{Could not create an appropriate shared temp folder},
-		action => q{incValue("errors"); diagnostic("$1", "error", 0, 0);}
-	},
-	{
-		id => "noClientLog",
-		pattern => q{Client exited without creating a log file},
-		action => q{incValue("errors"); diagnostic("$1", "error", 0, 0);}
+		action => q{incValue("errors"); diagnostic("Critical Error", "error", 0, forwardTo(q{Executing StaticShutdownAfterError}));}
 	},
 	{
 		id => "UATException",
 		pattern => q{ERROR: Exception in AutomationTool:},
-		action => q{incValue("errors"); diagnostic("$1", "error", 0, 0);}
+		action => q{incValue("errors"); diagnostic("Exception", "error", 0, 0);}
 	},
 	{
 		id => "LNK",
 		pattern => q{LINK : fatal error},
-		action => q{incValue("errors"); diagnostic("$1", "error", 0, 0);}
-	},
-	{
-		id => "NoLog",
-		pattern => q{Client exited without creating a log file},
-		action => action => q{incValue("errors"); diagnostic("$1", "error", 0, 0);}
-	},
-	{
-		id => "link",
-		pattern => q{Error executing},
-		action => q{incValue("errors"); diagnostic("$1", "error", -1, 0);}
-	},
-	{
-		id => "failedToMove",
-		pattern => q{Exception in AutomationTool: Failed to rename/move file},
-		action => q{incValue("errors"); diagnostic("$1", "error", 0, 0);}
-	},
-	{
-		id => "MSCORLIB",
-		pattern => q{Exception in mscorlib},
-		action => q{incValue("errors"); diagnostic("$1", "error", 0, 0);}
-	},
-	{
-		id => "ObjectRef",
-		pattern => q{Object reference not set to an instance of an object},
-		action => q{incValue("errors"); diagnostic("$1", "error", 0, 0);}
-	},
-	{
-		id => "XGEFailed",
-		pattern => q{XGE failed on try 1},
-		action => q{incValue("errors"); diagnostic("$1", "error", -5, 0);}
+		action => q{incValue("errors"); diagnostic("Link", "error", 0, 0);}
 	},
 	{
 		id => "UHTfailed",
 		pattern => q{UnrealHeaderTool failed for target},
-		action => q{incValue("errors"); diagnostic("$1", "error", 0, 0);}
-	},
-	{
-		id => "InvalidAsset",
-		pattern => q{The following asset pack names are invalid},
-		q{incValue("errors"); diagnostic("$1", "error", 0, forwardTo(q{StackTrace}));}
-	},
-	{
-		id => "Distill",
-		pattern => q{Sample uproject should have already been distilled to the correct location},
-		action => q{incValue("errors"); diagnostic("$1", "error", 0, 0);}
-	},
-	{
-		id => "ExitWithoutLog",
-		pattern => q{Client exited without creating a log file},
-		action => q{incValue("errors"); diagnostic("$1", "error", 0, 0);}
-	},
-	{
-		id => "failedtoconvert",
-		pattern => q{failed to convert due to errors},
-		action => q{incValue("errors"); diagnostic("$1", "error", 0, 0);}
+		action => q{incValue("errors"); diagnostic("UnrealHeaderTool", "error", 0, 0);}
 	},
 	{
 		id => "MSBuildError",
 		pattern => q{targets\(\d+,\d+\): error},
-		action => q{incValue("errors"); diagnostic("$1", "error", -2, 0);}
-	},
-	{
-		id => "ChangesSinceLastGreen",
-		pattern => q{PrintChanges:},
-		action => q{incValue("info"); diagnostic("$1", "info", 0, forwardTo(q{GUBP.PrintDetailedChanges: Took}));}
-	},
-	{
-		id => "Timeout",
-		pattern => q{ectool: ectool timeout: The command did not complete before the requested timeout},
-		action => q{incValue("errors"); diagnostic("$1", "error",-1, forwardTo(q{ExitCode=}));}
+		action => q{incValue("errors"); diagnostic("MSBuild", "error", -2, 0);}
 	},
 	{
 		id => "UnrealBuildToolException",
 		pattern => q{UnrealBuildTool Exception:},
-		action => q{incValue("errors"); diagnostic("$1", "error", 0, 5);}
+		action => q{incValue("errors"); diagnostic("UnrealBuildTool", "error", 0, 5);}
 	},
 	{
 		id => "ExitCode3",
 		pattern => q{ExitCode=3},
-		action => q{incValue("errors"); diagnostic("$1", "error", 0, 0);}
+		action => q{incValue("errors"); diagnostic("ExitCode", "error", 0, 0);}
 	},
 	{
 		id => "SegmentationFault",
 		pattern => q{ExitCode=139},
-		action => q{incValue("errors"); diagnostic("$1", "error", -4, 1);}
+		action => q{incValue("errors"); diagnostic("Exit code 139", "error", -4, 1);}
 	},
 	{
 		id => "ExitCode255",
 		pattern => q{ExitCode=255},
-		action => q{incValue("errors"); diagnostic("$1", "error", 0, 0);}
-	},
-	{
-		id => "MissingFileError",
-		pattern => q{build is missing or did not complete because this file is missing:},
-		action => q{incValue("errors"); diagnostic("$1", "error", 0, 1);}
-	},
-	{
-		id => "DerivedDataCacheError",
-		pattern => q{Failed while running DerivedDataCache},
-		action => q{incValue("errors"); diagnostic("$1", "error", 0, 1);}
+		action => q{incValue("errors"); diagnostic("Exit code 255", "error", 0, 0);}
 	},
 	{
 		id => "Stack dump",
 		pattern => q{Stack dump},
-		action => q{incValue("errors"); diagnostic("$1", "error", -4, forwardWhile('0', 1));}
+		action => q{incValue("errors"); diagnostic("Stack dump", "error", -4, forwardWhile('0', 1));}
 	},
 	{
 		id => "UATErrStack",
 		pattern => q{AutomationTool: Stack:},
-		action => q{incValue("errors"); diagnostic("$1", "error", 0, forwardTo(q{AutomationTool: \n}));}
+		action => q{incValue("errors"); diagnostic("AutomationTool", "error", 0, forwardTo(q{AutomationTool: \n}));}
 	},
 	{
 		id => "UATErrStack",
 		pattern => q{begin: stack for UAT},
-		action => q{incValue("errors"); diagnostic("$1", "error", 0, forwardTo(q{end: stack for UAT}));}
+		action => q{incValue("errors"); diagnostic("AutomationTool", "error", 0, forwardTo(q{end: stack for UAT}));}
 	},
 	{
 		id => "BuildFailed",
diff --git a/Engine/Build/BuildFarm/PostpFilter.pl b/Engine/Build/BuildFarm/PostpFilter.pl
--- a/Engine/Build/BuildFarm/PostpFilter.pl
+++ b/Engine/Build/BuildFarm/PostpFilter.pl
@@ -4,12 +4,18 @@ while(<>)
 {
 	# remove timestamps added by gubp
 	s/^\[[0-9:.]+\] //;
+
+	# remove the mono prefix for recursive UAT calls on Mac
+	s/^mono: ([A-Za-z0-9-]+: )/$1/;
+
+	# remove the AutomationTool prefix for recursive UAT calls anywhere
+	s/^AutomationTool: ([A-Za-z0-9-]+: )/$1/;
 	
 	# remove the gubp function or prefix
 	s/^([A-Za-z0-9.-]+): //;
 
 	# if it was the editor, also strip any additional timestamp
-	if($1 && $1 eq 'UE4Editor-Cmd')
+	if($1 && $1 =~ /^UE4Editor/)
 	{
 		s/^\[[0-9.:-]+\]\[\s*\d+\]//;
 	}
