Disable buffering on the postp pre-filter; was causing a delay in postp updating properties for diagnostic output, and preventing scripts from detecting failure conditions correctly.

diff --git a/Engine/Build/BuildFarm/PostpFilter.pl b/Engine/Build/BuildFarm/PostpFilter.pl
--- a/Engine/Build/BuildFarm/PostpFilter.pl
+++ b/Engine/Build/BuildFarm/PostpFilter.pl
@@ -1,5 +1,13 @@
 #!/usr/bin/env perl
 
+# make sure we don't buffer postp, otherwise we won't be able to read the list of warnings and errors back in for failure emails
+my $previous_handle = select(STDOUT);
+$| = 1;
+select(STDERR);
+$| = 1;
+select($previous_handle);
+
+# read everything from stdin
 while(<>)
 {
 	# remove timestamps added by gubp
