Fix UAT exceptions not being caught and written to the EC diagnostics file, meaning that they're not included in failure emails. Launcher builds are currently reporting no errors despite a failure exit code.

diff --git a/Engine/Build/BuildFarm/PostpExtensions.pl b/Engine/Build/BuildFarm/PostpExtensions.pl
--- a/Engine/Build/BuildFarm/PostpExtensions.pl
+++ b/Engine/Build/BuildFarm/PostpExtensions.pl
@@ -102,7 +102,7 @@ unshift @::gMatchers, (
 	},
 	{
 		id =>               "automationException",
-		pattern =>          q{^.* AutomationTool.AutomationException: },
+		pattern =>          q{AutomationTool\\.AutomationException: },
 		action =>           q{incValue("errors"); diagnostic("Exception", "error", 0, forwardWhile("^  at "));}
 	}
 );
diff --git a/Engine/Build/BuildFarm/PostpFilter.pl b/Engine/Build/BuildFarm/PostpFilter.pl
--- a/Engine/Build/BuildFarm/PostpFilter.pl
+++ b/Engine/Build/BuildFarm/PostpFilter.pl
@@ -20,7 +20,7 @@ while(<>)
 	s/^AutomationTool: ([A-Za-z0-9-]+: )/$1/;
 	
 	# remove the gubp function or prefix
-	s/^([A-Za-z0-9.-]+): //;
+	s/^([A-Za-z0-9.-]+): // if !/^AutomationTool\\.AutomationException:/;
 
 	# if it was the editor, also strip any additional timestamp
 	if($1 && $1 =~ /^UE4Editor/)
