EC: Fix a few glitches in build failure emails:

* Mac steps not highlighting submitters to files containing errors.
* Missing gap between error and warning diagnostics since changes to improve rendering on mobile devices.
* Prevent spaces being lost which are inserted for formatting in the compiler output.

diff --git a/Engine/Build/BuildFarm/PostpExtensions.pl b/Engine/Build/BuildFarm/PostpExtensions.pl
--- a/Engine/Build/BuildFarm/PostpExtensions.pl
+++ b/Engine/Build/BuildFarm/PostpExtensions.pl
@@ -77,7 +77,7 @@ unshift @::gMatchers, (
     },
     {
         id =>               "ubtFailedToProduceItem",
-        pattern =>          q{ERROR: UBT ERROR: Failed to produce item: },
+        pattern =>          q{(ERROR: )?UBT ERROR: Failed to produce item: },
 		action =>           q{incValue("errors"); diagnostic("UnrealBuildTool", "error")}
     },
     {
