EC: Ignore PS4 deprecation warnings on the build farm.

diff --git a/Engine/Build/BuildFarm/PostpExtensions.pl b/Engine/Build/BuildFarm/PostpExtensions.pl
--- a/Engine/Build/BuildFarm/PostpExtensions.pl
+++ b/Engine/Build/BuildFarm/PostpExtensions.pl
@@ -19,6 +19,10 @@
 	# SignTool errors which we can retry and recover from
 	"SignTool Error: ",
 	
+	# Temporary PS4 deprecation warnings
+	".*\\\\OnlineSubsystemPS4\\\\.*warning:.*\\[-Wdeprecated-declarations\\]",
+	".*\\\\PS4\\\\PS4Application\\.cpp.*warning:.*\\[-Wdeprecated-declarations\\]"
+	
 	
 #	".*ERROR: The process.*not found",
 #	".*ERROR: This operation returned because the timeout period expired.*",
