EC: Less restrictive regex to exclude even more PS4 warnings.

diff --git a/Engine/Build/BuildFarm/PostpExtensions.pl b/Engine/Build/BuildFarm/PostpExtensions.pl
--- a/Engine/Build/BuildFarm/PostpExtensions.pl
+++ b/Engine/Build/BuildFarm/PostpExtensions.pl
@@ -20,8 +20,8 @@
 	"SignTool Error: ",
 	
 	# Temporary PS4 deprecation warnings
-	".*\\\\OnlineSubsystemPS4\\\\.*warning:.*\\[-Wdeprecated-declarations\\]",
-	".*\\\\PS4\\\\PS4Application\\.cpp.*warning:.*\\[-Wdeprecated-declarations\\]"
+	".*OnlineSubsystemPS4.*warning:.*\\[-Wdeprecated-declarations\\]",
+	".*PS4Application\\.cpp.*warning:.*\\[-Wdeprecated-declarations\\]"
 	
 	
 #	".*ERROR: The process.*not found",
