Build: Suppress EC warnings when Doxygen outputs the name of a symbol called XXXWarning::YYY.

diff --git a/Engine/Build/BuildFarm/PostpExtensions.pl b/Engine/Build/BuildFarm/PostpExtensions.pl
--- a/Engine/Build/BuildFarm/PostpExtensions.pl
+++ b/Engine/Build/BuildFarm/PostpExtensions.pl
@@ -21,8 +21,10 @@
 	
 	# Temporary PS4 deprecation warnings
 	".*OnlineSubsystemPS4.*warning:.*\\[-Wdeprecated-declarations\\]",
-	".*PS4Application\\.cpp.*warning:.*\\[-Wdeprecated-declarations\\]"
+	".*PS4Application\\.cpp.*warning:.*\\[-Wdeprecated-declarations\\]",
 	
+	# Some doxygen output can confuse the post-processor, because it lists a symbol containing "Warning::"
+	"doxygen>.*Warning::.*",
 	
 #	".*ERROR: The process.*not found",
 #	".*ERROR: This operation returned because the timeout period expired.*",
@@ -171,7 +173,7 @@ unshift @::gMatchers, (
     },
     {
         id =>               "genericWarning",
-        pattern =>          q{WARNING:|[Ww]arning:},
+        pattern =>          q{WARNING:|[Ww]arning:(?!:)},
         action =>           q{incValue("warnings"); diagnostic("", "warning", 0)},
     },
 );
