EC: Fix a couple of patterns for parsing documentation errors.

diff --git a/Engine/Build/BuildFarm/PostpExtensions.pl b/Engine/Build/BuildFarm/PostpExtensions.pl
--- a/Engine/Build/BuildFarm/PostpExtensions.pl
+++ b/Engine/Build/BuildFarm/PostpExtensions.pl
@@ -29,7 +29,7 @@
 	"LogOnline:Warning: STEAM:.*libraries not present.*failed to load!",
 	
 	# Some doxygen output can confuse the post-processor, because it lists a symbol containing "Warning::"
-	"doxygen>.*Warning::.*",
+	"doxygen>.*(Warning|Error)::.*",
 	
 #	".*ERROR: The process.*not found",
 #	".*ERROR: This operation returned because the timeout period expired.*",
@@ -121,6 +121,11 @@ unshift @::gMatchers, (
         pattern =>          q{([^:]+):[\d:]+ warning:},
         action =>           q{incValue("warnings"); diagnostic($1, "warning", backWhile(": In function"), 0)},
     },
+    {
+        id =>               "genericDoctoolError",
+        pattern =>          q{Error:},
+        action =>           q{incValue("errors"); diagnostic("", "error")}
+    },
     {
         id =>               "ubtFailedToProduceItem",
         pattern =>          q{(ERROR: )?UBT ERROR: Failed to produce item: },
