EC: Ignore warnings for any build steps with a NotifyOnWarnings="false" attribute in BuildGraph. Should reduce noise on dashboard for projects we don't intend to clean up.

diff --git a/Engine/Build/BuildFarm/PostpFilter.pl b/Engine/Build/BuildFarm/PostpFilter.pl
--- a/Engine/Build/BuildFarm/PostpFilter.pl
+++ b/Engine/Build/BuildFarm/PostpFilter.pl
@@ -7,11 +7,14 @@ select(STDERR);
 $| = 1;
 select($previous_handle);
 
+# check for a flag to filter out warning lines
+my $no_warnings = grep(/--no-warnings/, @ARGV);
+
 # keep track of whether to output messages
 my $output_enabled = 1;
 
 # read everything from stdin
-while(<>)
+while(<STDIN>)
 {
 	# remove timestamps added by gubp
 	s/^\[[0-9:.]+\] //;
@@ -31,6 +34,12 @@ while(<>)
 		s/^\[[0-9.:-]+\]\[\s*\d+\]//;
 	}
 
+	# skip over warnings if necessary
+	if($no_warnings)
+	{
+		s/(?<![A-Za-z])[Ww]arning://g;
+	}
+
 	# look for a special marker to disable output
 	$output_enabled-- if /<-- Suspend Log Parsing -->/i;
 
