EC: Simplify rules for parsing UBT-style warnings and errors to include more lines of context for exceptions.

diff --git a/Engine/Build/BuildFarm/PostpExtensions.pl b/Engine/Build/BuildFarm/PostpExtensions.pl
--- a/Engine/Build/BuildFarm/PostpExtensions.pl
+++ b/Engine/Build/BuildFarm/PostpExtensions.pl
@@ -177,12 +177,12 @@ unshift @::gMatchers, (
 	{
 		id =>				"ubtError",
 		pattern =>			q{^ERROR:},
-		action =>			q{incValue("errors"); diagnostic("AutomationTool", "error", 0, forwardWhile("^       [^ ]"))}
+		action =>			q{incValue("errors"); diagnostic("AutomationTool", "error", 0, forwardWhile("^  "))}
 	},
 	{
 		id =>				"ubtWarning",
 		pattern =>			q{^WARNING:},
-		action =>			q{incValue("warnings"); diagnostic("AutomationTool", "warning", 0, forwardWhile("^         [^ ]"))}
+		action =>			q{incValue("warnings"); diagnostic("AutomationTool", "warning", 0, forwardWhile("^  "))}
 	},
     {
         id =>               "genericError",
