Ignore license warnings from PVS-Studio.

diff --git a/Engine/Build/BuildFarm/PostpExtensions.pl b/Engine/Build/BuildFarm/PostpExtensions.pl
--- a/Engine/Build/BuildFarm/PostpExtensions.pl
+++ b/Engine/Build/BuildFarm/PostpExtensions.pl
@@ -57,6 +57,9 @@
     "\\[Warn\\]	Number of Initial Chunks for Disc was changed by the patch/remaster\\. \\(Scenario #0, 0 -> 1\\) Confirm that it is intentional change\\.",
 	"\\[Warn\\]	Language Setting was changed by the patch/remaster\\. \\(Chunk #\\d+, 0x\\d+ -> 0x\\d+\\) Confirm that it is intentional change\\.",
 	
+	# PVS Studio
+	"warning Renew: Your license will expire in.*viva64\\.com",	
+	
 #	".*ERROR: The process.*not found",
 #	".*ERROR: This operation returned because the timeout period expired.*",
 #	".*Sync.VerifyKnownFileInManifest: ERROR:.*",
