EC: Add pattern matcher for Mono internal errors.

diff --git a/Engine/Build/BuildFarm/PostpExtensions.pl b/Engine/Build/BuildFarm/PostpExtensions.pl
--- a/Engine/Build/BuildFarm/PostpExtensions.pl
+++ b/Engine/Build/BuildFarm/PostpExtensions.pl
@@ -230,6 +230,11 @@ unshift @::gMatchers, (
 		id =>				"orbisPubCmdError",
         pattern =>          q{\[[Ee]rror\]},
 		action =>           q{incValue("errors"); diagnostic("", "error", 0, forwardWhile('^\\s*\\[[Ee]rror\\]')) },
+	},
+	{
+		id =>				"monoInternalError",
+		pattern =>			q{\* Assertion at wthreads.c:},
+		action =>           q{incValue("errors"); diagnostic("", "error", 0)},
 	}
 );
 
